Fixed crash when switching clothing meshes when clothing LOD0 is unbound on a clothing data object.

diff --git a/Engine/Source/Runtime/ClothingSystemRuntime/Private/ClothingSimulationNv.cpp b/Engine/Source/Runtime/ClothingSystemRuntime/Private/ClothingSimulationNv.cpp
--- a/Engine/Source/Runtime/ClothingSystemRuntime/Private/ClothingSimulationNv.cpp
+++ b/Engine/Source/Runtime/ClothingSystemRuntime/Private/ClothingSimulationNv.cpp
@@ -225,13 +225,9 @@ void FClothingSimulationNv::CreateActor(USkeletalMeshComponent* InOwnerComponent
 	// Pull collisions from the specified physics asset inside the clothing asset
 	ExtractActorCollisions(Asset, NewActor);
 
-	// Always start at zero, we'll pick up the right one before we simulate on the first tick
-	CurrentMeshLodIndex = 0;
-	NewActor.CurrentLodIndex = 0;
-
-	check(NewActor.LodData.IsValidIndex(0));
-
-	Solver->addCloth(NewActor.LodData[0].Cloth);
+	// Invalid indices so the call to UpdateLod runs all the correct logic as if our LOD just changed.
+	CurrentMeshLodIndex = INDEX_NONE;
+	NewActor.CurrentLodIndex = INDEX_NONE;
 
 	// Force update LODs so we're in the correct state now
 	UpdateLod(InOwnerComponent->PredictedLODLevel, InOwnerComponent->GetComponentTransform(), InOwnerComponent->GetComponentSpaceTransforms(), true);
