Merging CL# 4115124 from Release-4.20
--
Fixed clothing LOD update receiving empty bone transform array on initialization if a master pose component was used
#rnx

diff --git a/Engine/Source/Runtime/ClothingSystemRuntime/Private/ClothingSimulationNv.cpp b/Engine/Source/Runtime/ClothingSystemRuntime/Private/ClothingSimulationNv.cpp
--- a/Engine/Source/Runtime/ClothingSystemRuntime/Private/ClothingSimulationNv.cpp
+++ b/Engine/Source/Runtime/ClothingSystemRuntime/Private/ClothingSimulationNv.cpp
@@ -229,8 +229,9 @@ void FClothingSimulationNv::CreateActor(USkeletalMeshComponent* InOwnerComponent
 	CurrentMeshLodIndex = INDEX_NONE;
 	NewActor.CurrentLodIndex = INDEX_NONE;
 
-	// Force update LODs so we're in the correct state now
-	UpdateLod(InOwnerComponent->PredictedLODLevel, InOwnerComponent->GetComponentTransform(), InOwnerComponent->GetComponentSpaceTransforms(), true);
+	// Force update LODs so we're in the correct state now, need to resolve MPC if one is present
+	USkinnedMeshComponent* TransformComponent = InOwnerComponent->MasterPoseComponent.IsValid() ? InOwnerComponent->MasterPoseComponent.Get() : InOwnerComponent;
+	UpdateLod(InOwnerComponent->PredictedLODLevel, InOwnerComponent->GetComponentTransform(), TransformComponent->GetComponentSpaceTransforms(), true);
 
 	// Compute normals for all active actors for first frame
 	for(FClothingActorNv& Actor : Actors)
