Integrate as edit CL#4112806:
- Cloth crash for from Michael Lentine

diff --git a/Engine/Source/Runtime/ClothingSystemRuntime/Private/ClothingSimulation.cpp b/Engine/Source/Runtime/ClothingSystemRuntime/Private/ClothingSimulation.cpp
--- a/Engine/Source/Runtime/ClothingSystemRuntime/Private/ClothingSimulation.cpp
+++ b/Engine/Source/Runtime/ClothingSystemRuntime/Private/ClothingSimulation.cpp
@@ -164,7 +164,7 @@ void FClothingSimulationBase::FillContext(USkeletalMeshComponent* InComponent, f
 			    {
 				    const int32 MasterIndex = InComponent->MasterBoneMap[BoneIndex];
     
-				    if(MasterIndex != INDEX_NONE)
+				    if(MasterIndex != INDEX_NONE && MasterIndex < MasterComponent->GetComponentSpaceTransforms().Num())
 				    {
 					    BaseContext->BoneTransforms[BoneIndex] = MasterComponent->GetComponentSpaceTransforms()[MasterIndex];
 					    bFoundMaster = true;
diff --git a/Engine/Source/Runtime/ClothingSystemRuntime/Private/ClothingSimulationNv.cpp b/Engine/Source/Runtime/ClothingSystemRuntime/Private/ClothingSimulationNv.cpp
--- a/Engine/Source/Runtime/ClothingSystemRuntime/Private/ClothingSimulationNv.cpp
+++ b/Engine/Source/Runtime/ClothingSystemRuntime/Private/ClothingSimulationNv.cpp
@@ -802,7 +802,8 @@ void FClothingSimulationNv::GetSimulationData(TMap<int32, FClothSimulData>& OutD
 
 			if(!ReadTransformArray.IsValidIndex(Asset->ReferenceBoneIndex))
 			{
-				ensureMsgf(false, TEXT("Failed to write back clothing simulation data for component % as bone transforms are invalid."), *InOwnerComponent->GetName());
+				UE_LOG(LogSkeletalMesh, Warning, TEXT("Failed to write back clothing simulation data for component % as bone transforms are invalid."), *InOwnerComponent->GetName());
+				//ensureMsgf(false, TEXT("Failed to write back clothing simulation data for component % as bone transforms are invalid."), *InOwnerComponent->GetName());
 
 				ClothData.Reset();
 
