Dont access LOD data that doesn't exist.

diff --git a/Engine/Source/Runtime/ClothingSystemRuntime/Private/ClothingSimulationNv.cpp b/Engine/Source/Runtime/ClothingSystemRuntime/Private/ClothingSimulationNv.cpp
--- a/Engine/Source/Runtime/ClothingSystemRuntime/Private/ClothingSimulationNv.cpp
+++ b/Engine/Source/Runtime/ClothingSystemRuntime/Private/ClothingSimulationNv.cpp
@@ -67,7 +67,7 @@ void FClothingSimulationNv::CreateActor(USkeletalMeshComponent* InOwnerComponent
 
 	// Need the current reftolocals so we can skin the ref pose for the sim mesh
 	TArray<FMatrix> RefToLocals;
-	InOwnerComponent->GetCurrentRefToLocalMatrices(RefToLocals, InOwnerComponent->PredictedLODLevel);
+	InOwnerComponent->GetCurrentRefToLocalMatrices(RefToLocals, FMath::Min(InOwnerComponent->PredictedLODLevel, Asset->LodData.Num() - 1));
 
 	Actors.AddDefaulted();
 	FClothingActorNv& NewActor = Actors.Last();
