Fix asset that have unaccessible cloth binding

diff --git a/Engine/Source/Runtime/ClothingSystemRuntime/Private/Assets/ClothingAsset.cpp b/Engine/Source/Runtime/ClothingSystemRuntime/Private/Assets/ClothingAsset.cpp
--- a/Engine/Source/Runtime/ClothingSystemRuntime/Private/Assets/ClothingAsset.cpp
+++ b/Engine/Source/Runtime/ClothingSystemRuntime/Private/Assets/ClothingAsset.cpp
@@ -339,28 +339,26 @@ void UClothingAsset::UnbindFromSkeletalMesh(USkeletalMesh* InSkelMesh, int32 InM
 
 		FSkeletalMeshLODModel& LodModel = Mesh->LODModels[InMeshLodIndex];
 
-			for(int32 SectionIdx = LodModel.Sections.Num() - 1; SectionIdx >= 0; --SectionIdx)
-			{
+		for(int32 SectionIdx = LodModel.Sections.Num() - 1; SectionIdx >= 0; --SectionIdx)
+		{
 			FSkelMeshSection& Section = LodModel.Sections[SectionIdx];
-
 			if(Section.HasClothingData() && Section.ClothingData.AssetGuid == AssetGuid)
-					{
+			{
 				if(!bChangedMesh)
-						{
+				{
 					InSkelMesh->PreEditChange(nullptr);
-					}
-
+				}
 				ClothingAssetUtils::ClearSectionClothingData(Section);
-
-				// Clear the LOD map entry for this asset
-				if(LodMap.IsValidIndex(InMeshLodIndex))
-					{
-					LodMap[InMeshLodIndex] = INDEX_NONE;
-						}
-
-					bChangedMesh = true;
+				bChangedMesh = true;
 			}
 		}
+
+		// Clear the LOD map entry for this asset LOD, after a unbind we must be able to bind any asset
+		if (LodMap.IsValidIndex(InMeshLodIndex))
+		{
+			LodMap[InMeshLodIndex] = INDEX_NONE;
+			bChangedMesh = true;
+		}
 	}
 
 	// If the mesh changed we need to re-register any components that use it to reflect the changes
diff --git a/Engine/Source/Runtime/Engine/Private/SkeletalMesh.cpp b/Engine/Source/Runtime/Engine/Private/SkeletalMesh.cpp
--- a/Engine/Source/Runtime/Engine/Private/SkeletalMesh.cpp
+++ b/Engine/Source/Runtime/Engine/Private/SkeletalMesh.cpp
@@ -1658,6 +1658,36 @@ void USkeletalMesh::PostLoad()
 		{
 			CacheDerivedData();
 		}
+
+		//Make sure unused cloth are unbind
+		if (MeshClothingAssets.Num() > 0)
+		{
+			TArray<UClothingAssetBase *> InUsedClothingAssets;
+			GetClothingAssetsInUse(InUsedClothingAssets);
+			//Look if we have some cloth binding to unbind
+			for (UClothingAssetBase* MeshClothingAsset : MeshClothingAssets)
+			{
+				if (MeshClothingAsset == nullptr)
+				{
+					continue;
+				}
+				bool bFound = false;
+				for (UClothingAssetBase* UsedMeshClothingAsset : InUsedClothingAssets)
+				{
+					if (UsedMeshClothingAsset->GetAssetGuid() == MeshClothingAsset->GetAssetGuid())
+					{
+						bFound = true;
+						break;
+					}
+				}
+				if (!bFound)
+				{
+					//Make sure the asset is unbind, some old code path was allowing to have bind cloth asset not present in the imported model.
+					//The old inline reduction code was not rebinding the cloth asset nor unbind it.
+					MeshClothingAsset->UnbindFromSkeletalMesh(this);
+				}
+			}
+		}
 	}
 #endif // WITH_EDITOR
 
