Fix cloth popping

diff --git a/Engine/Source/Runtime/ClothingSystemRuntime/Private/ClothingSimulationNv.cpp b/Engine/Source/Runtime/ClothingSystemRuntime/Private/ClothingSimulationNv.cpp
--- a/Engine/Source/Runtime/ClothingSystemRuntime/Private/ClothingSimulationNv.cpp
+++ b/Engine/Source/Runtime/ClothingSystemRuntime/Private/ClothingSimulationNv.cpp
@@ -231,7 +231,7 @@ void FClothingSimulationNv::CreateActor(USkeletalMeshComponent* InOwnerComponent
 
 	// Force update LODs so we're in the correct state now, need to resolve MPC if one is present
 	USkinnedMeshComponent* TransformComponent = InOwnerComponent->MasterPoseComponent.IsValid() ? InOwnerComponent->MasterPoseComponent.Get() : InOwnerComponent;
-	UpdateLod(InOwnerComponent->PredictedLODLevel, InOwnerComponent->GetComponentTransform(), TransformComponent->GetComponentSpaceTransforms(), true, true);
+	UpdateLod(InOwnerComponent->PredictedLODLevel, InOwnerComponent->GetComponentTransform(), TransformComponent->GetComponentSpaceTransforms(), RefToLocals, true, true);
 
 	// Compute normals for all active actors for first frame
 	for(FClothingActorNv& Actor : Actors)
@@ -507,7 +507,7 @@ void FClothingSimulationNv::Simulate(IClothingSimulationContext* InContext)
 		return;
 	}
 
-	UpdateLod(NvContext->PredictedLod, NvContext->ComponentToWorld, NvContext->BoneTransforms);
+	UpdateLod(NvContext->PredictedLod, NvContext->ComponentToWorld, NvContext->BoneTransforms, NvContext->RefToLocals);
 
 	// Pre-sim work
 	for(FClothingActorNv& Actor : Actors)
@@ -975,7 +975,7 @@ FBoxSphereBounds FClothingSimulationNv::GetBounds(const USkeletalMeshComponent*
 	return CurrentBounds;
 }
 
-void FClothingSimulationNv::UpdateLod(int32 InPredictedLod, const FTransform& ComponentToWorld, const TArray<FTransform>& CSTransforms, bool bForceNoRemap, bool bForceActorChecks)
+void FClothingSimulationNv::UpdateLod(int32 InPredictedLod, const FTransform& ComponentToWorld, const TArray<FTransform>& CSTransforms, const TArray<FMatrix>& RefToLocals, bool bForceNoRemap, bool bForceActorChecks)
 {
 	if(InPredictedLod != CurrentMeshLodIndex || bForceActorChecks)
 	{
@@ -1105,6 +1105,15 @@ void FClothingSimulationNv::UpdateLod(int32 InPredictedLod, const FTransform& Co
 				NewLodData.Cloth->setRotation(U2PQuat(SimRootTransform.GetRotation()));
 				NewLodData.Cloth->clearInertia();
 
+				// clear spheres and planes double buffering since cloth kept collision data from last time it was simulated
+				// and clearInertia does not fully empty buffers
+				NewLodData.Cloth->setSpheres(nv::cloth::Range<const PxVec4>(), 0, NewLodData.Cloth->getNumSpheres()); // empty spheres
+				NewLodData.Cloth->setPlanes(nv::cloth::Range<const PxVec4>(), 0, NewLodData.Cloth->getNumPlanes()); // empty planes
+
+				// clear constraints from previous sim
+				NewLodData.Cloth->clearMotionConstraints();
+				NewLodData.Cloth->clearSeparationConstraints();
+
 				Actor.CurrentLodIndex = PredictedClothingLod;
 			}
 			else
@@ -1114,11 +1123,36 @@ void FClothingSimulationNv::UpdateLod(int32 InPredictedLod, const FTransform& Co
 				{
 					Solver->addCloth(NewLodData.Cloth);
 
-					for(int32 ParticleIndex = 0; ParticleIndex < NumNewParticles; ++ParticleIndex)
+					if(CSTransforms.Num() > Actor.AssetCreatedFrom->ReferenceBoneIndex)
 					{
-						NewLodParticles[ParticleIndex] = NewLodData.Px_RestPositions[ParticleIndex];
-						NewLodPrevParticles[ParticleIndex] = NewLodData.Px_RestPositions[ParticleIndex];
-						NewAccelerations[ParticleIndex] = physx::PxVec4(0.0f);
+						// compute skinned positions to init sim mesh
+						const FClothPhysicalMeshData& PhysMesh = Actor.AssetCreatedFrom->LodData[PredictedClothingLod].PhysicalMeshData;
+						TArray<FVector> SkinnedPhysicsMeshPositions;
+						TArray<FVector> SkinnedPhysicsMeshNormals;
+						FClothingSimulationBase::SkinPhysicsMesh(
+							Actor.AssetCreatedFrom,
+							PhysMesh,
+							CSTransforms[Actor.AssetCreatedFrom->ReferenceBoneIndex],
+							RefToLocals.GetData(),
+							RefToLocals.Num(),
+							SkinnedPhysicsMeshPositions,
+							SkinnedPhysicsMeshNormals);
+
+						for(int32 ParticleIndex = 0; ParticleIndex < NumNewParticles; ++ParticleIndex)
+						{
+							NewLodParticles[ParticleIndex] = PxVec4(U2PVector(SkinnedPhysicsMeshPositions[ParticleIndex]), NewLodParticles[ParticleIndex].w);
+							NewLodPrevParticles[ParticleIndex] = PxVec4(U2PVector(SkinnedPhysicsMeshPositions[ParticleIndex]), NewLodPrevParticles[ParticleIndex].w);
+							NewAccelerations[ParticleIndex] = physx::PxVec4(0.0f);
+						}
+					}
+					else
+					{
+						for(int32 ParticleIndex = 0; ParticleIndex < NumNewParticles; ++ParticleIndex)
+						{
+							NewLodParticles[ParticleIndex] = NewLodData.Px_RestPositions[ParticleIndex];
+							NewLodPrevParticles[ParticleIndex] = NewLodData.Px_RestPositions[ParticleIndex];
+							NewAccelerations[ParticleIndex] = physx::PxVec4(0.0f);
+						}
 					}
 
 					FTransform SimRootTransform = CSTransforms[Actor.AssetCreatedFrom->ReferenceBoneIndex] * ComponentToWorld;
@@ -1126,6 +1160,15 @@ void FClothingSimulationNv::UpdateLod(int32 InPredictedLod, const FTransform& Co
 					NewLodData.Cloth->setRotation(U2PQuat(SimRootTransform.GetRotation()));
 					NewLodData.Cloth->clearInertia();
 
+					// clear spheres and planes double buffering since cloth kept collision data from last time it was simulated
+					// and clearInertia does not fully empty buffers
+					NewLodData.Cloth->setSpheres(nv::cloth::Range<const PxVec4>(), 0, NewLodData.Cloth->getNumSpheres()); // empty spheres
+					NewLodData.Cloth->setPlanes(nv::cloth::Range<const PxVec4>(), 0, NewLodData.Cloth->getNumPlanes()); // empty planes
+
+					// clear constraints from previous sim
+					NewLodData.Cloth->clearMotionConstraints();
+					NewLodData.Cloth->clearSeparationConstraints();
+
 					Actor.CurrentLodIndex = PredictedClothingLod;
 				}
 				else
diff --git a/Engine/Source/Runtime/ClothingSystemRuntime/Public/ClothingSimulationNv.h b/Engine/Source/Runtime/ClothingSystemRuntime/Public/ClothingSimulationNv.h
--- a/Engine/Source/Runtime/ClothingSystemRuntime/Public/ClothingSimulationNv.h
+++ b/Engine/Source/Runtime/ClothingSystemRuntime/Public/ClothingSimulationNv.h
@@ -219,9 +219,10 @@ private:
 	// @param InPredictedLod The predicted LOD for the mesh component this clothing simulation is running on
 	// @param ComponentToWorld The component to world transform for the mesh component this clothing simulation is running on
 	// @param CSTranforms Component space transforms of the owning skeletal mesh component
+	// @param RefToLocals RefToLocals of the owning skeletal mesh component to init simulation mesh
 	// @param bForceNoRemap When changing LODs The incoming LOD can be skinned to the outgoing LOD for better transitions - this flag skips that step
 	// @param bForceActorChecks Typically we trust all LODs to be in sync, but that isn't always the case (e.g. adding a new actor). This forces each actor's LOD to be checked instead of trusting the master LOD
-	void UpdateLod(int32 InPredictedLod, const FTransform& ComponentToWorld, const TArray<FTransform>& CSTransforms, bool bForceNoRemap = false, bool bForceActorChecks = false);
+	void UpdateLod(int32 InPredictedLod, const FTransform& ComponentToWorld, const TArray<FTransform>& CSTransforms, const TArray<FMatrix>& RefToLocals, bool bForceNoRemap = false, bool bForceActorChecks = false);
 
 	// The core simulation is only solving unoriented particles, so we need to compute normals after the
 	// simulation runs
diff --git a/Engine/Source/Runtime/Engine/Private/SkeletalMeshComponentPhysics.cpp b/Engine/Source/Runtime/Engine/Private/SkeletalMeshComponentPhysics.cpp
--- a/Engine/Source/Runtime/Engine/Private/SkeletalMeshComponentPhysics.cpp
+++ b/Engine/Source/Runtime/Engine/Private/SkeletalMeshComponentPhysics.cpp
@@ -2820,7 +2820,10 @@ void USkeletalMeshComponent::ProcessClothCollisionWithEnvironment()
 
 					if(SkelComp->ClothingSimulation)
 					{
-						SkelComp->ClothingSimulation->GetCollisions(NewCollisionData, false);
+						// append skeletal component collisions
+						FClothCollisionData SkelCollisionData;
+						SkelComp->ClothingSimulation->GetCollisions(SkelCollisionData, false);
+						NewCollisionData.Append(SkelCollisionData);
 					}
 				}
 			}
