Fixing crash for cloth sim when clothing section has clothing set for one LOD profile, but mesh provides differnet predicted LOD when Forcing LOD is disabled.

diff --git a/Engine/Source/Runtime/ClothingSystemRuntime/Private/ClothingSimulationNv.cpp b/Engine/Source/Runtime/ClothingSystemRuntime/Private/ClothingSimulationNv.cpp
--- a/Engine/Source/Runtime/ClothingSystemRuntime/Private/ClothingSimulationNv.cpp
+++ b/Engine/Source/Runtime/ClothingSystemRuntime/Private/ClothingSimulationNv.cpp
@@ -1001,9 +1001,23 @@ void FClothingSimulationNv::UpdateLod(int32 InPredictedLod, const FTransform& Co
 			bool bOldLodMapped = LodMap.IsValidIndex(CurrentMeshLodIndex) && LodMap[CurrentMeshLodIndex] != INDEX_NONE;
 
 			// Get the clothing LOD mapped from the mesh predicted LOD
-			const int32 PredictedClothingLod = LodMap[InPredictedLod];
+
+
 			const int32 OldClothingLod = bOldLodMapped ? LodMap[CurrentMeshLodIndex] : INDEX_NONE;
 
+			// If potentialLod doesn't map to a valid LOD, we try higher LOD levels for a valid LOD.
+			// Asset might only have lod on LOD 1 and not 0, however if mesh doesn't force LOD to 1, 
+			// asset will not be assigned valid LOD index and will not generate sim data, breaking things.
+			int32 PredictedClothingLod = INDEX_NONE;
+			for (int32 PotentialLod = InPredictedLod; PotentialLod < LodMap.Num(); ++PotentialLod)
+			{
+				if (LodMap[PotentialLod] != INDEX_NONE)
+				{
+					PredictedClothingLod = LodMap[PotentialLod];
+					break;
+				}
+			}
+
 			if(PredictedClothingLod == Actor.CurrentLodIndex)
 			{
 				// We must have forced a LOD update because we added a new actor - this actor is good to go though.
